CIS - Fixed Issue 1157809: Errors in Incremental Compile Shipping Editor and Incremental Compile Editor caqused by duplicate definitions.

#Synced-CL 48812260

diff --git a/Engine/Plugins/ChaosClothAssetDataflowNodes/Source/ChaosClothAssetDataflowNodes/Public/ChaosClothAsset/USDImportNode_v2.h b/Engine/Plugins/ChaosClothAssetDataflowNodes/Source/ChaosClothAssetDataflowNodes/Public/ChaosClothAsset/USDImportNode_v2.h
--- a/Engine/Plugins/ChaosClothAssetDataflowNodes/Source/ChaosClothAssetDataflowNodes/Public/ChaosClothAsset/USDImportNode_v2.h
+++ b/Engine/Plugins/ChaosClothAssetDataflowNodes/Source/ChaosClothAssetDataflowNodes/Public/ChaosClothAsset/USDImportNode_v2.h
@@ -23,8 +23,8 @@ struct FChaosClothAssetUsdClothData final
 	TMap<FName, int32> SimPatternFabricIndices;
 	FManagedArrayCollection SimulationCollection;
 
-	void Reset();
-	bool Serialize(FArchive& Ar);
+	CHAOSCLOTHASSETDATAFLOWNODES_API void Reset();
+	CHAOSCLOTHASSETDATAFLOWNODES_API bool Serialize(FArchive& Ar);
 };
 
 template<> struct TStructOpsTypeTraits<FChaosClothAssetUsdClothData> : public TStructOpsTypeTraitsBase2<FChaosClothAssetUsdClothData>
diff --git a/Engine/Plugins/ChaosClothAssetUsdDataflowNodes/Source/ChaosClothAssetUsdDataflowNodes/Private/ChaosClothAsset/USDImportNode_v3.cpp b/Engine/Plugins/ChaosClothAssetUsdDataflowNodes/Source/ChaosClothAssetUsdDataflowNodes/Private/ChaosClothAsset/USDImportNode_v3.cpp
--- a/Engine/Plugins/ChaosClothAssetUsdDataflowNodes/Source/ChaosClothAssetUsdDataflowNodes/Private/ChaosClothAsset/USDImportNode_v3.cpp
+++ b/Engine/Plugins/ChaosClothAssetUsdDataflowNodes/Source/ChaosClothAssetUsdDataflowNodes/Private/ChaosClothAsset/USDImportNode_v3.cpp
@@ -1684,34 +1684,6 @@ namespace UE::Chaos::ClothAsset::Private
 #endif  // #if USE_USD_SDK
 }  // End namespace Private
 
-void FChaosClothAssetUsdClothData::Reset()
-{
-	SimPatterns.Reset();
-	Sewings.Reset();
-	RenderPatterns.Reset();
-	RenderToSimPatterns.Reset();
-	SimPatternFabricIndices.Reset();
-	SimulationCollection.Reset();
-}
-
-bool FChaosClothAssetUsdClothData::Serialize(FArchive& Ar)
-{
-	Ar.UsingCustomVersion(FUE5ReleaseStreamObjectVersion::GUID);
-
-	Ar << SimPatterns;
-	Ar << Sewings;
-	Ar << RenderPatterns;
-	Ar << RenderToSimPatterns;
-
-	if (Ar.CustomVer(FUE5ReleaseStreamObjectVersion::GUID) >= FUE5ReleaseStreamObjectVersion::AddSimulationPropertySupportToClothUSDImportNodeV2)
-	{
-		Ar << SimPatternFabricIndices;
-		SimulationCollection.Serialize(Ar);
-	}
-	
-	return true;
-}
-
 FChaosClothAssetUSDImportNode_v3::FChaosClothAssetUSDImportNode_v3(const UE::Dataflow::FNodeParameters& InParam, FGuid InGuid)
 	: FDataflowNode(InParam, InGuid)
 	, UsdFile(
